[ChaosVD] Initial Implementation of CVD as a standalone program

Now CVD can be built as an unreal program (in addition to being an editor tool).

diff --git a/Engine/Plugins/ChaosVD/ChaosVD.uplugin b/Engine/Plugins/ChaosVD/ChaosVD.uplugin
--- a/Engine/Plugins/ChaosVD/ChaosVD.uplugin
+++ b/Engine/Plugins/ChaosVD/ChaosVD.uplugin
@@ -18,7 +18,7 @@
 	"Modules": [
 		{
 			"Name": "ChaosVD",
-			"Type": "Editor",
+			"Type": "EditorAndProgram",
 			"LoadingPhase": "Default"
 		}
 	],
@@ -27,5 +27,6 @@
 			"Name": "GeometryProcessing",
 			"Enabled": true
 		}
-	]
+	],
+	"SupportedPrograms": ["ChaosVisualDebugger"]
 }
\ No newline at end of file
diff --git a/Engine/Plugins/ChaosVD/Source/ChaosVD/Private/ChaosVDModule.cpp b/Engine/Plugins/ChaosVD/Source/ChaosVD/Private/ChaosVDModule.cpp
--- a/Engine/Plugins/ChaosVD/Source/ChaosVD/Private/ChaosVDModule.cpp
+++ b/Engine/Plugins/ChaosVD/Source/ChaosVD/Private/ChaosVDModule.cpp
@@ -3,20 +3,22 @@
 #include "ChaosVDModule.h"
 
 #include "ChaosVDCommands.h"
-#include "ChaosVDStyle.h"
+#include "ChaosVDEngine.h"
 #include "ChaosVDParticleActorCustomization.h"
+#include "ChaosVDStyle.h"
 #include "ChaosVDTabsIDs.h"
-#include "Widgets/Docking/SDockTab.h"
-#include "Widgets/SChaosVDMainTab.h"
-#include "PropertyEditorModule.h"
-#include "ChaosVDEngine.h"
-#include "WorkspaceMenuStructure.h"
-#include "WorkspaceMenuStructureModule.h"
 #include "DetailsCustomizations/ChaosVDGeometryComponentCustomization.h"
 #include "DetailsCustomizations/ChaosVDParticleDataWrapperCustomization.h"
 #include "DetailsCustomizations/ChaosVDQueryDataWrappersCustomizationDetails.h"
-#include "Trace/ChaosVDTraceManager.h"
+#include "Misc/App.h"
+#include "Misc/CommandLine.h"
 #include "Misc/Guid.h"
+#include "PropertyEditorModule.h"
+#include "Trace/ChaosVDTraceManager.h"
+#include "Widgets/Docking/SDockTab.h"
+#include "Widgets/SChaosVDMainTab.h"
+#include "WorkspaceMenuStructure.h"
+#include "WorkspaceMenuStructureModule.h"
 
 #define LOCTEXT_NAMESPACE "ChaosVisualDebugger"
 
@@ -31,6 +33,8 @@ FAutoConsoleCommand ChaosVDSpawnNewCVDInstance(
 	})
 );
 
+FString FChaosVDModule::ChaosVisualDebuggerProgramName = TEXT("ChaosVisualDebugger");
+
 FChaosVDModule& FChaosVDModule::Get()
 {
 	return FModuleManager::Get().LoadModuleChecked<FChaosVDModule>(TEXT("ChaosVD"));
@@ -52,6 +56,13 @@ void FChaosVDModule::StartupModule()
 
 	ChaosVDTraceManager = MakeShared<FChaosVDTraceManager>();
 
+	if (IsStandaloneChaosVisualDebugger())
+	{
+		// In the standalone app, once the engine is initialized we need to spawn the main tab otherwise there will be no UI
+		// because we intentionally don't load the mainframe / rest of the editor UI
+		FCoreDelegates::OnFEngineLoopInitComplete.AddRaw(this, &FChaosVDModule::SpawnCVDTab);
+	}
+
 	FCoreDelegates::OnEnginePreExit.AddRaw(this, &FChaosVDModule::CloseActiveInstances);
 }
 
@@ -67,7 +78,12 @@ void FChaosVDModule::ShutdownModule()
 	{
 		FGlobalTabmanager::Get()->UnregisterNomadTabSpawner(TabID);
 	}
-	
+
+	if (IsStandaloneChaosVisualDebugger())
+	{
+		FCoreDelegates::OnFEngineLoopInitComplete.RemoveAll(this);
+	}
+
 	FCoreDelegates::OnEnginePreExit.RemoveAll(this);
 
 	CloseActiveInstances();
@@ -90,6 +106,12 @@ void FChaosVDModule::RegisterClassesCustomDetails() const
 
 void FChaosVDModule::SpawnCVDTab()
 {
+	if (IsStandaloneChaosVisualDebugger())
+	{
+		// In the standalone app, we need to load the status bar module so the status bar subsystem is initialized
+		FModuleManager::Get().LoadModule("StatusBar");
+	}
+
 	// Registering new tab spawners with random names is not the best idea to spawn new tabs, but it is good enough to test we can run multiple instances of CVD withing the editor.
 	// This is also why spawning new tabs it is only exposed via console commands for now.
 	// When this feature is deemed stable and exposed in the UI I will investigate a more correct way of implementing this
@@ -136,6 +158,12 @@ TSharedRef<SDockTab> FChaosVDModule::SpawnMainTab(const FSpawnTabArgs& Args)
 
 void FChaosVDModule::HandleTabClosed(TSharedRef<SDockTab> ClosedTab, FGuid InstanceGUID)
 {
+	if (IsStandaloneChaosVisualDebugger())
+	{
+		// If this is the standalone CVD app, we can assume that tab closed indicates an exit request
+		RequestEngineExit(TEXT("MainCVDTabClosed"));
+	}
+
 	// Workaround. Currently the ChaosVD Engine instance determines the lifetime of the Editor world and other objects
 	// Some widgets, like UE Level viewport tries to iterate on these objects on destruction
 	// For now we can avoid any crashes by just de-initializing ChaosVD Engine on the next frame but that is not the real fix.
@@ -208,7 +236,11 @@ void FChaosVDModule::CloseActiveInstances()
 	ActiveCVDTabs.Reset();
 }
 
+bool FChaosVDModule::IsStandaloneChaosVisualDebugger()
+{
+	return FPlatformProperties::IsProgram() && FApp::GetProjectName() == ChaosVisualDebuggerProgramName;
+}
 
 #undef LOCTEXT_NAMESPACE
 	
-IMPLEMENT_MODULE(FChaosVDModule, ChaosVD)
\ No newline at end of file
+IMPLEMENT_MODULE(FChaosVDModule, ChaosVD)
diff --git a/Engine/Plugins/ChaosVD/Source/ChaosVD/Private/ChaosVDPlaybackViewportClient.cpp b/Engine/Plugins/ChaosVD/Source/ChaosVD/Private/ChaosVDPlaybackViewportClient.cpp
--- a/Engine/Plugins/ChaosVD/Source/ChaosVD/Private/ChaosVDPlaybackViewportClient.cpp
+++ b/Engine/Plugins/ChaosVD/Source/ChaosVD/Private/ChaosVDPlaybackViewportClient.cpp
@@ -40,6 +40,7 @@ FChaosVDPlaybackViewportClient::FChaosVDPlaybackViewportClient(const TSharedPtr<
 	{
 		Settings->OnFarClippingOverrideChanged().AddRaw(this, &FChaosVDPlaybackViewportClient::HandleViewportSettingsChanged);
 		Settings->OnVisibilitySettingsChanged().AddRaw(this, &FChaosVDPlaybackViewportClient::HandleViewportSettingsChanged);
+		Settings->OnColorSettingsChanged().AddRaw(this, &FChaosVDPlaybackViewportClient::HandleViewportSettingsChanged);
 
 		HandleViewportSettingsChanged(Settings);
 	}
@@ -64,6 +65,7 @@ FChaosVDPlaybackViewportClient::~FChaosVDPlaybackViewportClient()
 	{
 		Settings->OnFarClippingOverrideChanged().RemoveAll(this);
 		Settings->OnVisibilitySettingsChanged().RemoveAll(this);
+		Settings->OnColorSettingsChanged().RemoveAll(this);
 	}
 }
 
@@ -176,6 +178,7 @@ void FChaosVDPlaybackViewportClient::HandleActorMoving(AActor* MovedActor) const
 		{
 			if (SceneSharedPtr->GetSkySphereActor()->Implements<UChaosVDSkySphereInterface>())
 			{
+				FEditorScriptExecutionGuard AllowEditorScriptGuard;
 				IChaosVDSkySphereInterface::Execute_Refresh(SceneSharedPtr->GetSkySphereActor());
 			}
 		}
diff --git a/Engine/Plugins/ChaosVD/Source/ChaosVD/Private/ChaosVDScene.cpp b/Engine/Plugins/ChaosVD/Source/ChaosVD/Private/ChaosVDScene.cpp
--- a/Engine/Plugins/ChaosVD/Source/ChaosVD/Private/ChaosVDScene.cpp
+++ b/Engine/Plugins/ChaosVD/Source/ChaosVD/Private/ChaosVDScene.cpp
@@ -472,6 +472,7 @@ void FChaosVDScene::CreateBaseLights(UWorld* TargetWorld) const
 				SkySphere->SetFolderPath(LightingFolderPath);
 				if (SkySphere->Implements<UChaosVDSkySphereInterface>())
 				{
+					FEditorScriptExecutionGuard AllowEditorScriptGuard;
 					IChaosVDSkySphereInterface::Execute_SetDirectionalLightSource(SkySphere, DirectionalLightActor);
 				}
 			}
diff --git a/Engine/Plugins/ChaosVD/Source/ChaosVD/Private/ChaosVDWorldOutlinerMode.cpp b/Engine/Plugins/ChaosVD/Source/ChaosVD/Private/ChaosVDWorldOutlinerMode.cpp
--- a/Engine/Plugins/ChaosVD/Source/ChaosVD/Private/ChaosVDWorldOutlinerMode.cpp
+++ b/Engine/Plugins/ChaosVD/Source/ChaosVD/Private/ChaosVDWorldOutlinerMode.cpp
@@ -199,6 +199,14 @@ TUniquePtr<ISceneOutlinerHierarchy> FChaosVDWorldOutlinerMode::CreateHierarchy()
 	return ActorHierarchy;
 }
 
+bool FChaosVDWorldOutlinerMode::CanInteract(const ISceneOutlinerTreeItem& Item) const
+{
+	// This option is not supported in CVD yet
+	ensure(!bCanInteractWithSelectableActorsOnly);
+
+	return true;
+}
+
 void FChaosVDWorldOutlinerMode::EnqueueAndCombineHierarchyEvent(const FSceneOutlinerTreeItemID& ItemID, const FSceneOutlinerHierarchyChangedData& EnventToProcess)
 {
 	if (FSceneOutlinerHierarchyChangedData* EventData = PendingOutlinerEventsMap.Find(ItemID))
diff --git a/Engine/Plugins/ChaosVD/Source/ChaosVD/Private/ChaosVDWorldOutlinerMode.h b/Engine/Plugins/ChaosVD/Source/ChaosVD/Private/ChaosVDWorldOutlinerMode.h
--- a/Engine/Plugins/ChaosVD/Source/ChaosVD/Private/ChaosVDWorldOutlinerMode.h
+++ b/Engine/Plugins/ChaosVD/Source/ChaosVD/Private/ChaosVDWorldOutlinerMode.h
@@ -95,6 +95,8 @@ public:
 
 	virtual TUniquePtr<ISceneOutlinerHierarchy> CreateHierarchy() override;
 
+	bool CanInteract(const ISceneOutlinerTreeItem& Item) const override;
+
 private:
 
 	void EnqueueAndCombineHierarchyEvent(const FSceneOutlinerTreeItemID& ItemID, const FSceneOutlinerHierarchyChangedData& EnventToProcess);
diff --git a/Engine/Plugins/ChaosVD/Source/ChaosVD/Private/ChaosVDWorldOutlinerTab.cpp b/Engine/Plugins/ChaosVD/Source/ChaosVD/Private/ChaosVDWorldOutlinerTab.cpp
--- a/Engine/Plugins/ChaosVD/Source/ChaosVD/Private/ChaosVDWorldOutlinerTab.cpp
+++ b/Engine/Plugins/ChaosVD/Source/ChaosVD/Private/ChaosVDWorldOutlinerTab.cpp
@@ -31,6 +31,7 @@ void FChaosVDWorldOutlinerTab::CreateWorldOutlinerWidget()
 		FActorModeParams ModeParams(Outliner);
 		ModeParams.SpecifiedWorldToDisplay = GetChaosVDWorld();
 		ModeParams.bHideEmptyFolders = true;
+		ModeParams.bCanInteractWithSelectableActorsOnly = false;
 
 		// The mode is deleted by the Outliner when it is destroyed 
 		return new FChaosVDWorldOutlinerMode(ModeParams, GetChaosVDScene());
diff --git a/Engine/Plugins/ChaosVD/Source/ChaosVD/Private/Widgets/SChaosVDRecordingControls.cpp b/Engine/Plugins/ChaosVD/Source/ChaosVD/Private/Widgets/SChaosVDRecordingControls.cpp
--- a/Engine/Plugins/ChaosVD/Source/ChaosVD/Private/Widgets/SChaosVDRecordingControls.cpp
+++ b/Engine/Plugins/ChaosVD/Source/ChaosVD/Private/Widgets/SChaosVDRecordingControls.cpp
@@ -215,7 +215,7 @@ bool SChaosVDRecordingControls::CanChangeChannelEnabledState(TWeakPtr<FCVDDataCh
 bool SChaosVDRecordingControls::HasDataChannelsSupport() const
 {
 #if WITH_CHAOS_VISUAL_DEBUGGER
-	return true;
+	return !FChaosVDModule::IsStandaloneChaosVisualDebugger();
 #else
 	return false;
 #endif
@@ -395,6 +395,13 @@ bool SChaosVDRecordingControls::IsRecordingToggleButtonEnabled(EChaosVDRecording
 		return false;
 	}
 
+	// Currently we don't support start a recording in standalone mode
+	// but once we add a way to communicate with other running instances we can re-enable these buttons
+	if (FChaosVDModule::IsStandaloneChaosVisualDebugger())
+	{
+		return false;
+	}
+
 	const TSharedPtr<SChaosVDMainTab> MainTabSharedPtr = MainTabWeakPtr.Pin();
 	if (!MainTabSharedPtr.IsValid())
 	{
diff --git a/Engine/Plugins/ChaosVD/Source/ChaosVD/Public/ChaosVDModule.h b/Engine/Plugins/ChaosVD/Source/ChaosVD/Public/ChaosVDModule.h
--- a/Engine/Plugins/ChaosVD/Source/ChaosVD/Public/ChaosVDModule.h
+++ b/Engine/Plugins/ChaosVD/Source/ChaosVD/Public/ChaosVDModule.h
@@ -30,6 +30,8 @@ public:
 	/** Spawns a new instance of CVD */
 	void SpawnCVDTab();
 
+	static bool IsStandaloneChaosVisualDebugger();
+
 private:
 
 	void RegisterClassesCustomDetails() const;
@@ -54,5 +56,7 @@ private:
 
 	TSharedPtr<FChaosVDTraceManager> ChaosVDTraceManager;
 
+	static FString ChaosVisualDebuggerProgramName;
+
 	bool bIsShuttingDown = false;
 };
diff --git a/Engine/Plugins/Runtime/GeometryProcessing/GeometryProcessing.uplugin b/Engine/Plugins/Runtime/GeometryProcessing/GeometryProcessing.uplugin
--- a/Engine/Plugins/Runtime/GeometryProcessing/GeometryProcessing.uplugin
+++ b/Engine/Plugins/Runtime/GeometryProcessing/GeometryProcessing.uplugin
@@ -24,5 +24,6 @@
 			"Type": "Runtime",
 			"LoadingPhase": "Default"
 		}
-	]
+	],
+	"SupportedPrograms": ["ChaosVisualDebugger"]
 }
\ No newline at end of file
diff --git a/Engine/Plugins/Runtime/PropertyAccess/PropertyAccessEditor.uplugin b/Engine/Plugins/Runtime/PropertyAccess/PropertyAccessEditor.uplugin
--- a/Engine/Plugins/Runtime/PropertyAccess/PropertyAccessEditor.uplugin
+++ b/Engine/Plugins/Runtime/PropertyAccess/PropertyAccessEditor.uplugin
@@ -22,5 +22,6 @@
 			"Type": "UncookedOnly" ,
 			"LoadingPhase" : "PreDefault"
 		}		
-	]
+	],
+	"SupportedPrograms": ["ChaosVisualDebugger"]
 }
\ No newline at end of file
diff --git a/Engine/Programs/ChaosVisualDebugger/BuildAndCook.bat b/Engine/Programs/ChaosVisualDebugger/BuildAndCook.bat
new file mode 100755
--- /dev/null
+++ b/Engine/Programs/ChaosVisualDebugger/BuildAndCook.bat
@@ -0,0 +1,3 @@
+echo off
+call "%~dp0..\..\Build\BatchFiles\RunUAT.bat" MakeCookedEditor -nop4 -utf8output -project="%~dp0/ChaosVisualDebugger.uproject" -DDC=noshared -stage -archive -package -build -pak -iostore -compressed -prereqs -archivedirectory="%~dp0PackagedBuild" -clientconfig="Development" -cook -nocompile -nocompileuat -AdditionalCookerOptions=" -CookProcessCount=4" -unrealexe="ChaosVisualDebugger.exe" %*
+exit /b !ERRORLEVEL!
\ No newline at end of file
diff --git a/Engine/Programs/ChaosVisualDebugger/ChaosVisualDebugger.uproject b/Engine/Programs/ChaosVisualDebugger/ChaosVisualDebugger.uproject
new file mode 100644
--- /dev/null
+++ b/Engine/Programs/ChaosVisualDebugger/ChaosVisualDebugger.uproject
@@ -0,0 +1,19 @@
+{
+	"FileVersion": 3,
+	"EngineAssociation": "",
+	"Category": "",
+	"Description": "",
+	"TargetPlatforms": [
+		"WindowsCookedEditor"
+	],
+	"Plugins": [
+		{
+			"Name": "ChaosVD",
+			"Enabled": true
+		},
+		{
+			"Name": "PropertyAccessEditor",
+			"Enabled": true
+		}	
+	]
+}
\ No newline at end of file
diff --git a/Engine/Programs/ChaosVisualDebugger/Config/DefaultEngine.ini b/Engine/Programs/ChaosVisualDebugger/Config/DefaultEngine.ini
new file mode 100644
--- /dev/null
+++ b/Engine/Programs/ChaosVisualDebugger/Config/DefaultEngine.ini
@@ -0,0 +1,29 @@
+[/Script/Engine.GarbageCollectionSettings]
+gc.MaxObjectsInProgram=500000
+gc.MaxObjectsInGame=500000
+gc.MaxObjectsInEditor=500000
+
+[SystemSettingsEditor]
+s.AllowUnversionedContentInEditor=1
+cook.AllowCookedDataInEditorBuilds=1
+s.EnforcePackageCompatibleVersionCheck=0
+
+[SystemSettings]
+s.AllowUnversionedContentInEditor=1
+cook.AllowCookedDataInEditorBuilds=1
+s.EnforcePackageCompatibleVersionCheck=0
+
+[/Script/AndroidFileServerEditor.AndroidFileServerRuntimeSettings]
+bEnablePlugin=True
+bAllowNetworkConnection=True
+SecurityToken=D4CD916B4290446E57EEBCB1B9F9E4A2
+bIncludeInShipping=False
+bAllowExternalStartInShipping=False
+bCompileAFSProject=False
+bUseCompression=False
+bLogFiles=False
+bReportStats=False
+ConnectionType=USBOnly
+bUseManualIPAddress=False
+ManualIPAddress=
+
diff --git a/Engine/Programs/ChaosVisualDebugger/Config/DefaultGame.ini b/Engine/Programs/ChaosVisualDebugger/Config/DefaultGame.ini
new file mode 100644
--- /dev/null
+++ b/Engine/Programs/ChaosVisualDebugger/Config/DefaultGame.ini
@@ -0,0 +1,28 @@
+[/Script/UnrealEd.ProjectPackagingSettings]
+bUseZenStore=false
+
+[Staging]
++DisallowedLocalizationTargets=Engine
+
+[CookedEditorSettings]
+bSupportCookedEditor=true
+CookedEditorTargetName=ChaosVisualDebugger
+bStageUAT=false
+bStageShaderDirs=true
+bStagePlatformBuildDirs=false
+bStageExtrasDirs=false
+bStagePlatformDirs=false
+bIsForExternalDistribution=true
+MapMode=None
++EngineExtraStageFiles=(Path="Plugins/ChaosVD/Resources", Recursive=false, NonUFS=true)
++EngineExtraStageFiles=(Path="Content/Splash", Recursive=false, NonUFS=true)
+
+; remove everything in Platforms, except and Config files that aren't in Restricted, and shaders
++FinalStripFiles=(Filter="/Platforms/", Exceptions="(Engine)/Platforms/.*/Config || /Shaders/ || Android\.Automation || /XB1ZLib/")
++FinalStripFiles=(Filter="Engine/Extras")
+
+; this will remove all of the compiled UAT script dlls, except for some ones that had their own subdirectories when built, that we need to keep
++FinalStripFiles=(Filter="AutomationTool/AutomationScripts/", Exceptions="/(BuildGraph|BuildPatchTool|ContentWorker|Scripts|Android)/")
+
+; Always cook the control rig standard function library
++EngineSpecificAssetsToCook="/ControlRig/StandardFunctionLibrary/StandardFunctionLibrary.uasset"
diff --git a/Engine/Source/Programs/ChaosVisualDebugger/ChaosVisualDebugger.Build.cs b/Engine/Source/Programs/ChaosVisualDebugger/ChaosVisualDebugger.Build.cs
--- a/Engine/Source/Programs/ChaosVisualDebugger/ChaosVisualDebugger.Build.cs
+++ b/Engine/Source/Programs/ChaosVisualDebugger/ChaosVisualDebugger.Build.cs
@@ -6,37 +6,73 @@ public class ChaosVisualDebugger : ModuleRules
 {
 	public ChaosVisualDebugger(ReadOnlyTargetRules Target) : base(Target)
 	{
-		PublicIncludePathModuleNames.Add("Launch");
-		PrivateDependencyModuleNames.Add("Core");
-		PrivateDependencyModuleNames.Add("Projects");
+		PrivateDependencyModuleNames.Add("ChaosVD");
 
+		// LaunchEngineLoop dependencies
 		PrivateDependencyModuleNames.AddRange(
 			new string[] {
-				"AppFramework",
-				"Core",
 				"ApplicationCore",
+				"Core",
+				"CoreUObject",
+				"Engine",
+				"InputCore",
+				"InstallBundleManager",
+				"MediaUtils",
+				"Messaging",
+				"MoviePlayer",
+				"MoviePlayerProxy",
 				"Projects",
+				"PreLoadScreen",
+				"PIEPreviewDeviceProfileSelector",
+				"RenderCore",
+				"RHI",
 				"Slate",
 				"SlateCore",
 				"StandaloneRenderer",
-				//"CoreUObject",
-				"SourceCodeAccess",
-				"TraceInsights",
-				"TraceAnalysis",
-				"TraceServices"
+				"TraceLog",
+				"ProfilerService",
+				"ProfileVisualizer",
+				"PropertyAccessEditor"
 			}
 		);
 
+		// LaunchEngineLoop IncludePath dependencies
 		PrivateIncludePathModuleNames.AddRange(
 			new string[] {
-				"SlateReflector",
+				"Launch",
+				"AutomationWorker",
+				"AutomationController",
+				"AutomationTest",
+				"DerivedDataCache",
+				"HeadMountedDisplay", 
+				"MRMesh", 
+				"SlateRHIRenderer", 
+				"SlateNullRenderer",
 			}
 		);
 
-		DynamicallyLoadedModuleNames.AddRange(
+		// LaunchEngineLoop editor dependencies
+		PrivateDependencyModuleNames.AddRange(
 			new string[] {
-				"SlateReflector",
-			}
-		);
+				
+				"PropertyEditor",
+				"DerivedDataCache",
+				"ToolWidgets",
+				"UnrealEd"
+			});
+
+		if (Target.IsInPlatformGroup(UnrealPlatformGroup.Linux))
+		{
+			PrivateDependencyModuleNames.AddRange(
+				new string[] {
+					"UnixCommonStartup"
+				}
+			);
+		}
+
+		if (Target.Platform == UnrealTargetPlatform.Win64 && Target.Architecture == UnrealArch.X64)
+		{
+			PrivateDefinitions.Add("USE_D3D12_REDIST=1");
+		}
 	}
 }
diff --git a/Engine/Source/Programs/ChaosVisualDebugger/ChaosVisualDebugger.Target.cs b/Engine/Source/Programs/ChaosVisualDebugger/ChaosVisualDebugger.Target.cs
--- a/Engine/Source/Programs/ChaosVisualDebugger/ChaosVisualDebugger.Target.cs
+++ b/Engine/Source/Programs/ChaosVisualDebugger/ChaosVisualDebugger.Target.cs
@@ -11,19 +11,55 @@ public class ChaosVisualDebuggerTarget : TargetRules
 		Type = TargetType.Program;
 		LinkType = TargetLinkType.Monolithic;
 		LaunchModuleName = "ChaosVisualDebugger";
+		BuildEnvironment = TargetBuildEnvironment.Unique;
 
-		// Lean and mean
+		DefaultBuildSettings = BuildSettingsVersion.Latest;
+		IncludeOrderVersion = EngineIncludeOrderVersion.Latest;
+		
+		bIsBuildingConsoleApplication = false;
+		
+		bCompileAgainstCoreUObject = true;
+		bCompileAgainstEngine = true;
+		bCompileAgainstEditor = true;
+		bBuildWithEditorOnlyData = true;
+		bUsesSlate = true;
+		bEnableTrace = true;
+		bUseLoggingInShipping = true;
+
+		bCompileWithPluginSupport = true;
 		bBuildDeveloperTools = true;
+		
+		// Disabled whatever we can and we don't need
+		bCompilePython = false;
+		bCompileICU = true;
+		bCompileRecast = false;
+		bUseGameplayDebugger = false;
+		bCompileNavmeshClusterLinks = false;
+		bCompileNavmeshSegmentLinks = false;
+		bWithLiveCoding = false;
+		bIncludePluginsForTargetPlatforms = false;
+		
+		OptedInModulePlatforms = new UnrealTargetPlatform[] { UnrealTargetPlatform.Win64 };
 
-		// Editor-only data, however, is needed
-		bBuildWithEditorOnlyData = true;
+		AdditionalPlugins.AddRange(new string[]
+		{
+			"ChaosVD",
+			"GeometryProcessing",
+			"PropertyAccessEditor"
+		});
 
-		// Currently this app is not linking against the engine, so we'll compile out references from Core to the rest of the engine
-		bCompileAgainstEngine = false;
-		bCompileAgainstCoreUObject = true;
-		//bCompileAgainstApplicationCore = false;
+		EnablePlugins.Add("GeometryProcessing");
+		EnablePlugins.Add("ChaosVD");
+		EnablePlugins.Add("PropertyAccessEditor");
 
-		// This a Windows app (sets entry point to WinMain(), instead of main())
-		bIsBuildingConsoleApplication = false;
+		GlobalDefinitions.Add("AUTOSDKS_ENABLED=0");
+		GlobalDefinitions.Add("ASSETREGISTRY_ENABLE_PREMADE_REGISTRY_IN_EDITOR=1");
+		GlobalDefinitions.Add("UE_IS_COOKED_EDITOR=1");
+		GlobalDefinitions.Add("UE_FORCE_USE_PAKS=1");
+		GlobalDefinitions.Add("UE_FORCE_USE_IOSTORE=1");
+
+		// this will allow shader compiling to work based on whether or not the shaders directory is present
+		// to determine if we should allow shader compilation
+		GlobalDefinitions.Add("UE_ALLOW_SHADER_COMPILING_BASED_ON_SHADER_DIRECTORY_EXISTENCE=1");
 	}
 }
diff --git a/Engine/Source/Programs/ChaosVisualDebugger/Private/ChaosVisualDebuggerMain.cpp b/Engine/Source/Programs/ChaosVisualDebugger/Private/ChaosVisualDebuggerMain.cpp
--- a/Engine/Source/Programs/ChaosVisualDebugger/Private/ChaosVisualDebuggerMain.cpp
+++ b/Engine/Source/Programs/ChaosVisualDebugger/Private/ChaosVisualDebuggerMain.cpp
@@ -1,141 +1,81 @@
 // Copyright Epic Games, Inc. All Rights Reserved.
 
-
 #include "ChaosVisualDebuggerMain.h"
-#include "ChaosVisualDebuggerSlateStyle.h"
+
 #include "RequiredProgramMainCPPInclude.h"
-#include "StandaloneRenderer.h"
-#include "Widgets/Testing/STestSuite.h"
-#include "Modules/ModuleManager.h"
-#include "ISlateReflectorModule.h"
+#include "EditorViewportClient.h"
+#include "LaunchEngineLoop.h"
+#include "HAL/PlatformSplash.h"
 
-DEFINE_LOG_CATEGORY_STATIC(LogChaosVisualDebugger, Log, All);
+IMPLEMENT_APPLICATION(ChaosVisualDebugger, "ChaosVisualDebugger");
 
-IMPLEMENT_APPLICATION(ChaosVisualDebugger, "Chaos Visual Debugger");
+DEFINE_LOG_CATEGORY_STATIC(LogChaosVisualDebugger, Log, All);
 
+// Opt in to new D3D12 redist and tell the loader where to search for D3D12Core.dll.
+// The D3D loader looks for these symbol exports in the .exe module.
+// We only support this on x64 Windows Desktop platforms. Other platforms or non-redist-aware 
+// versions of Windows will transparently load default OS-provided D3D12 library.
+#if USE_D3D12_REDIST
+extern "C" { _declspec(dllexport) extern const UINT D3D12SDKVersion = 611; } // D3D12_SDK_VERSION
+extern "C" { _declspec(dllexport) extern const char* D3D12SDKPath = ".\\D3D12\\"; }
+#endif // USE_D3D12_REDIST
 
-void InitializedSlateApplication()
+int32 RunChaosVisualDebugger(const TCHAR* CommandLine)
 {
-	// crank up a normal Slate application using the platform's standalone renderer
-	FSlateApplication::InitializeAsStandaloneApplication(GetStandardStandaloneRenderer());
+	FTaskTagScope Scope(ETaskTag::EGameThread);
 
-	// Set the application name.
-	const FText ApplicationTitle = NSLOCTEXT("ChaosVisualDebugger", "AppTitle", "ChaosVisualDebugger");
-	FGlobalTabmanager::Get()->SetApplicationTitle(ApplicationTitle);
-}
+#if !(UE_BUILD_SHIPPING)
 
-TSharedRef<SDockTab> SpawnViewport(const FSpawnTabArgs& Args)
-{
-	TSharedRef<SDockTab> ViewportTab =
-	SNew(SDockTab)
-	.TabRole(ETabRole::MajorTab)
-	.Label(NSLOCTEXT("ChaosVisualDebugger","ViewportTabLabel", "Viewport"))
-	.ToolTipText(NSLOCTEXT("ChaosVisualDebugger","ViewportTabToolTip", "The Chaos Visual debugger Viewport is under development"));
+	// If "-waitforattach" or "-WaitForDebugger" was specified, halt startup and wait for a debugger to attach before continuing
+	if (FParse::Param(CommandLine, TEXT("waitforattach")) || FParse::Param(CommandLine, TEXT("WaitForDebugger")))
+	{
+		while (!FPlatformMisc::IsDebuggerPresent())
+		{
+			FPlatformProcess::Sleep(0.1f);
+		}
+		UE_DEBUG_BREAK();
+	}
 
-
-	FVisualDebuggerStyle::ResetToDefault();	
-	
-	ViewportTab->SetContent
-	(
-		SNew(SVerticalBox)
-		+ SVerticalBox::Slot()
-		.AutoHeight()
-		.HAlign(HAlign_Center)
-		.Padding(30)
-		[
-			SNew(SImage)
-			.Image(FVisualDebuggerStyle::Get().GetBrush("UE4Icon"))
-		]
-		+ SVerticalBox::Slot()
-		.AutoHeight()
-		.HAlign(HAlign_Center)
-		[
-			SNew(STextBlock)
-			.Text(NSLOCTEXT("ChaosVisualDebugger", "TestBlock", "Visual Debugger, under development!"))
-		]
-	);
-
-	return ViewportTab;
-}
-
-
-void BuildChaosVDBUserInterface()
-{
-	// Need to load this module so we have the widget reflector tab available
-	FModuleManager::LoadModuleChecked<ISlateReflectorModule>("SlateReflector");
-
-	FGlobalTabmanager::Get()->RegisterTabSpawner("Viewport", FOnSpawnTab::CreateStatic(&SpawnViewport));
-
-	TSharedRef<FTabManager::FLayout> Layout = FTabManager::NewLayout("SlateVisualDebugger_Layout")
-		->AddArea
-		(
-			FTabManager::NewArea(1600, 1200)
-			->SetWindow(FVector2D(420, 10), false)
-			->Split
-			(
-				FTabManager::NewStack()
-				->AddTab("Viewport", ETabState::OpenedTab)
-			)
-		)
-		// This is for debugging the GUI
-		/*->AddArea
-		(
-			// This area will get a 400x600 window at 10,10
-			FTabManager::NewArea(400, 600)
-			->SetWindow(FVector2D(10, 10), false)
-			->Split
-			(
-				// The area contains a single tab with the widget reflector, for debugging purposes
-				FTabManager::NewStack()->AddTab("WidgetReflector", ETabState::OpenedTab)
-			)
-		)*/
-		;
-
-	FGlobalTabmanager::Get()->RestoreFrom(Layout, TSharedPtr<SWindow>());
-}
-
-
-int32 ChaosVisualDebuggerMain(const TCHAR* CommandLine)
-{
-	UE_LOG(LogChaosVisualDebugger, Display, TEXT("Chaos Visual Debugger - Early Prototype Development"));
+#endif
 
 	// Override the stack size for the thread pool.
 	FQueuedThreadPool::OverrideStackSize = 256 * 1024;
 
-	FCommandLine::Set(CommandLine);
-
-	// start up the main loop
-	GEngineLoop.PreInit(CommandLine);
+	const FText AppName = NSLOCTEXT("ChaosVisualDebugger", "ChaosVisualDebuggerSplashText", "Chaos Visual Debugger");
+	FPlatformSplash::SetSplashText(SplashTextType::GameName, *AppName.ToString());
 
-	// Make sure all UObject classes are registered and default properties have been initialized
-	ProcessNewlyLoadedUObjects();
+	FString Command;
+	const bool bIsRunningCommand = FParse::Value(CommandLine, TEXT("-RUN="), Command);
 
-	// Tell the module manager it may now process newly-loaded UObjects when new C++ modules are loaded
-	FModuleManager::Get().StartProcessingNewlyLoadedObjects();
+	const FString CommandLineString = CommandLine;
+	const FString FinalCommandLine = bIsRunningCommand ? CommandLine : CommandLineString + TEXT(" EDITOR");
 
-	InitializedSlateApplication();
+	// start up the main loop
+	const int32 Result = GEngineLoop.PreInit(*FinalCommandLine);
 
+	if (Result != 0)
 	{
-		BuildChaosVDBUserInterface();
-		// Bring up the test suite (for testing)
-		//RestoreSlateTestSuite();
+		UE_LOG(LogChaosVisualDebugger, Error, TEXT("EngineLoop PreInit failed!"));
+		return Result;
 	}
 
-	while (!IsEngineExitRequested())
+	if (!bIsRunningCommand)
 	{
-		FTaskGraphInterface::Get().ProcessThreadUntilIdle(ENamedThreads::GameThread);
-		FStats::AdvanceFrame(false);
-		FTSTicker::GetCoreTicker().Tick(FApp::GetDeltaTime());
-		FSlateApplication::Get().PumpMessages();
-		FSlateApplication::Get().Tick();
-		FPlatformProcess::Sleep(0);
-	}
+		// Register navigation commands for all viewports
+		FViewportNavigationCommands::Register();
 
-	FCoreDelegates::OnExit.Broadcast();
-	FSlateApplication::Shutdown();
-	FEngineLoop::AppPreExit();
-	FModuleManager::Get().UnloadModulesAtShutdown();
-	FEngineLoop::AppExit();
+		GEngineLoop.Init();
+
+		// Hide the splash screen now that everything is ready to go
+		FPlatformSplash::Hide();
+
+		while (!IsEngineExitRequested())
+		{
+			GEngineLoop.Tick();
+		}
+	}
+	
+	GEngineLoop.Exit();
 
-	return 0;
-} 
\ No newline at end of file
+	return Result;
+}
\ No newline at end of file
diff --git a/Engine/Source/Programs/ChaosVisualDebugger/Private/ChaosVisualDebuggerMain.h b/Engine/Source/Programs/ChaosVisualDebugger/Private/ChaosVisualDebuggerMain.h
--- a/Engine/Source/Programs/ChaosVisualDebugger/Private/ChaosVisualDebuggerMain.h
+++ b/Engine/Source/Programs/ChaosVisualDebugger/Private/ChaosVisualDebuggerMain.h
@@ -1,7 +1,8 @@
 // Copyright Epic Games, Inc. All Rights Reserved.
 
 #pragma once
+#include "HAL/Platform.h"
 
 // This is the platform independent entry point to the visual debugger
 // It returns the error level
-int32 ChaosVisualDebuggerMain(const TCHAR* CommandLine);
\ No newline at end of file
+int32 RunChaosVisualDebugger(const TCHAR* CommandLine);
\ No newline at end of file
diff --git a/Engine/Source/Programs/ChaosVisualDebugger/Private/ChaosVisualDebuggerSlateStyle.cpp b/Engine/Source/Programs/ChaosVisualDebugger/Private/ChaosVisualDebuggerSlateStyle.cpp
deleted file mode 100644
diff --git a/Engine/Source/Programs/ChaosVisualDebugger/Private/ChaosVisualDebuggerSlateStyle.h b/Engine/Source/Programs/ChaosVisualDebugger/Private/ChaosVisualDebuggerSlateStyle.h
deleted file mode 100644
diff --git a/Engine/Source/Programs/ChaosVisualDebugger/Private/Linux/ChaosVisualDebuggerLinux.cpp b/Engine/Source/Programs/ChaosVisualDebugger/Private/Linux/ChaosVisualDebuggerLinux.cpp
deleted file mode 100644
diff --git a/Engine/Source/Programs/ChaosVisualDebugger/Private/Mac/ChaosVisualDebuggerMac.cpp b/Engine/Source/Programs/ChaosVisualDebugger/Private/Mac/ChaosVisualDebuggerMac.cpp
deleted file mode 100644
diff --git a/Engine/Source/Programs/ChaosVisualDebugger/Private/Windows/ChaosVisualDebuggerWindows.cpp b/Engine/Source/Programs/ChaosVisualDebugger/Private/Windows/ChaosVisualDebuggerWindows.cpp
--- a/Engine/Source/Programs/ChaosVisualDebugger/Private/Windows/ChaosVisualDebuggerWindows.cpp
+++ b/Engine/Source/Programs/ChaosVisualDebugger/Private/Windows/ChaosVisualDebuggerWindows.cpp
@@ -1,10 +1,5 @@
 // Copyright Epic Games, Inc. All Rights Reserved.
 
-// Note: This application is still in very early development
-
-// This translation unit is windows specific
-
-#include "ChaosVisualDebuggerWindows.h"
 #include "ChaosVisualDebuggerMain.h"
 #include "HAL/ExceptionHandling.h"
 #include "Windows/WindowsHWrapper.h"
@@ -21,15 +16,13 @@
  * @param nShowCmd Specifies how the window is to be shown.
  * @return Application's exit value.
  */
-
-int WINAPI WinMain(_In_ HINSTANCE hInInstance, _In_opt_ HINSTANCE hPrevInstance, _In_ LPSTR lpCmdLine, _In_ int nCmdShow)
+int32 WINAPI WinMain(_In_ HINSTANCE hInInstance, _In_opt_ HINSTANCE hPrevInstance, _In_ LPSTR lpCmdLine, _In_ int nShowCmd)
 {
 	hInstance = hInInstance;
 
-	const TCHAR* CmdLine = ::GetCommandLineW();
+	const WIDECHAR* CmdLine = ::GetCommandLineW();
 	CmdLine = FCommandLine::RemoveExeName(CmdLine);
 
-
 	int32 ErrorLevel = 0;
 
 #if UE_BUILD_DEBUG
@@ -38,7 +31,7 @@ int WINAPI WinMain(_In_ HINSTANCE hInInstance, _In_opt_ HINSTANCE hPrevInstance,
 	if (FPlatformMisc::IsDebuggerPresent() && !GAlwaysReportCrash)
 #endif
 	{
-		ErrorLevel = ChaosVisualDebuggerMain(CmdLine);
+		ErrorLevel = RunChaosVisualDebugger(CmdLine);
 	}
 	else
 	{
@@ -46,9 +39,9 @@ int WINAPI WinMain(_In_ HINSTANCE hInInstance, _In_opt_ HINSTANCE hPrevInstance,
 		__try
 #endif
 		{
-			GIsGuarded = 1;
-			ErrorLevel = ChaosVisualDebuggerMain(CmdLine);
-			GIsGuarded = 0;
+			GIsGuarded = true;
+			ErrorLevel = RunChaosVisualDebugger(CmdLine);
+			GIsGuarded = false;
 		}
 #if !PLATFORM_SEH_EXCEPTIONS_DISABLED
 		__except (ReportCrash(GetExceptionInformation()))
diff --git a/Engine/Source/Programs/ChaosVisualDebugger/Private/Windows/ChaosVisualDebuggerWindows.h b/Engine/Source/Programs/ChaosVisualDebugger/Private/Windows/ChaosVisualDebuggerWindows.h
deleted file mode 100644
diff --git a/Engine/Source/Runtime/GeometryFramework/GeometryFramework.Build.cs b/Engine/Source/Runtime/GeometryFramework/GeometryFramework.Build.cs
--- a/Engine/Source/Runtime/GeometryFramework/GeometryFramework.Build.cs
+++ b/Engine/Source/Runtime/GeometryFramework/GeometryFramework.Build.cs
@@ -15,7 +15,7 @@ public class GeometryFramework : ModuleRules
 		);
 		
 		// These are required to register filters to the Level Editor Outliner
-		if (Target.Type == TargetType.Editor)
+		if (Target.Type == TargetType.Editor || Target.Type == TargetType.Program)
 		{
 			PrivateDependencyModuleNames.AddRange(
 				new string[]
