[ChaosVD] Making CVD standalone only supported on Windows platforms

This is just temporary until we can test it in other platforms.

diff --git a/Engine/Source/Programs/ChaosVisualDebugger/ChaosVisualDebugger.Target.cs b/Engine/Source/Programs/ChaosVisualDebugger/ChaosVisualDebugger.Target.cs
--- a/Engine/Source/Programs/ChaosVisualDebugger/ChaosVisualDebugger.Target.cs
+++ b/Engine/Source/Programs/ChaosVisualDebugger/ChaosVisualDebugger.Target.cs
@@ -3,7 +3,7 @@
 using UnrealBuildTool;
 using System.Collections.Generic;
 
-[SupportedPlatforms(UnrealPlatformClass.Desktop)]
+[SupportedPlatforms("Win64")]
 public class ChaosVisualDebuggerTarget : TargetRules
 {
 	public ChaosVisualDebuggerTarget(TargetInfo Target) : base(Target)
