[ChaosVD] Fixes for multiple issues during shutdown in the Standalone version of the Chaos Visual Debugger

It seems that recently a rendering (and other systems) refactor was done and that is causing issues with the the standalone version fo CVD. This CL implement workaround until the root causes can be fixed later on.

diff --git a/Engine/Programs/ChaosVisualDebugger/Config/DefaultEngine.ini b/Engine/Programs/ChaosVisualDebugger/Config/DefaultEngine.ini
--- a/Engine/Programs/ChaosVisualDebugger/Config/DefaultEngine.ini
+++ b/Engine/Programs/ChaosVisualDebugger/Config/DefaultEngine.ini
@@ -27,3 +27,5 @@ ConnectionType=USBOnly
 bUseManualIPAddress=False
 ManualIPAddress=
 
+[ConsoleVariables]
+r.Nanite.ProjectEnabled=0
diff --git a/Engine/Source/Editor/UnrealEd/Private/DerivedDataBuildWorkers.cpp b/Engine/Source/Editor/UnrealEd/Private/DerivedDataBuildWorkers.cpp
--- a/Engine/Source/Editor/UnrealEd/Private/DerivedDataBuildWorkers.cpp
+++ b/Engine/Source/Editor/UnrealEd/Private/DerivedDataBuildWorkers.cpp
@@ -21,3 +21,8 @@ void InitDerivedDataBuildWorkers()
 		GBuildWorkers.Emplace();
 	}
 }
+
+void ShutdownDerivedDataBuildWorkers()
+{
+	GBuildWorkers.Reset();
+}
diff --git a/Engine/Source/Editor/UnrealEd/Private/DerivedDataBuildWorkers.h b/Engine/Source/Editor/UnrealEd/Private/DerivedDataBuildWorkers.h
--- a/Engine/Source/Editor/UnrealEd/Private/DerivedDataBuildWorkers.h
+++ b/Engine/Source/Editor/UnrealEd/Private/DerivedDataBuildWorkers.h
@@ -3,3 +3,5 @@
 #pragma once
 
 void InitDerivedDataBuildWorkers();
+
+void ShutdownDerivedDataBuildWorkers();
diff --git a/Engine/Source/Editor/UnrealEd/Private/EditorEngine.cpp b/Engine/Source/Editor/UnrealEd/Private/EditorEngine.cpp
--- a/Engine/Source/Editor/UnrealEd/Private/EditorEngine.cpp
+++ b/Engine/Source/Editor/UnrealEd/Private/EditorEngine.cpp
@@ -1658,6 +1658,8 @@ void UEditorEngine::FinishDestroy()
 		}
 	}
 
+	ShutdownDerivedDataBuildWorkers();
+
 	Super::FinishDestroy();
 }
 
diff --git a/Engine/Source/Programs/ChaosVisualDebugger/Private/ChaosVisualDebuggerMain.cpp b/Engine/Source/Programs/ChaosVisualDebugger/Private/ChaosVisualDebuggerMain.cpp
--- a/Engine/Source/Programs/ChaosVisualDebugger/Private/ChaosVisualDebuggerMain.cpp
+++ b/Engine/Source/Programs/ChaosVisualDebugger/Private/ChaosVisualDebuggerMain.cpp
@@ -16,7 +16,7 @@ DEFINE_LOG_CATEGORY_STATIC(LogChaosVisualDebugger, Log, All);
 // We only support this on x64 Windows Desktop platforms. Other platforms or non-redist-aware 
 // versions of Windows will transparently load default OS-provided D3D12 library.
 #if USE_D3D12_REDIST
-extern "C" { _declspec(dllexport) extern const UINT D3D12SDKVersion = 611; } // D3D12_SDK_VERSION
+extern "C" { _declspec(dllexport) extern const UINT D3D12SDKVersion = 614; } // D3D12_SDK_VERSION
 extern "C" { _declspec(dllexport) extern const char* D3D12SDKPath = ".\\D3D12\\"; }
 #endif // USE_D3D12_REDIST
 
@@ -74,7 +74,10 @@ int32 RunChaosVisualDebugger(const TCHAR* CommandLine)
 			GEngineLoop.Tick();
 		}
 	}
-	
+
+	// Make sure all rendering thread work is done before we attempt to exit the program
+	FFrameEndSync::Sync(true);
+
 	GEngineLoop.Exit();
 
 	return Result;
