[ChaosVD] Fix for link error when clang is used to compile the standalone version of CVD.

There was a dependency that is no longer needed, and it was causing duplicated symbols error.

diff --git a/Engine/Source/Programs/ChaosVisualDebugger/ChaosVisualDebugger.Build.cs b/Engine/Source/Programs/ChaosVisualDebugger/ChaosVisualDebugger.Build.cs
--- a/Engine/Source/Programs/ChaosVisualDebugger/ChaosVisualDebugger.Build.cs
+++ b/Engine/Source/Programs/ChaosVisualDebugger/ChaosVisualDebugger.Build.cs
@@ -28,7 +28,6 @@ public class ChaosVisualDebugger : ModuleRules
 				"RHI",
 				"Slate",
 				"SlateCore",
-				"StandaloneRenderer",
 				"TraceLog",
 				"ProfileVisualizer",
 				"PropertyAccessEditor"
