Updated code to use the FPlatformSplash::SetSplashText overload that takes FText

#rnx

diff --git a/Engine/Source/Programs/ChaosVisualDebugger/Private/ChaosVisualDebuggerMain.cpp b/Engine/Source/Programs/ChaosVisualDebugger/Private/ChaosVisualDebuggerMain.cpp
--- a/Engine/Source/Programs/ChaosVisualDebugger/Private/ChaosVisualDebuggerMain.cpp
+++ b/Engine/Source/Programs/ChaosVisualDebugger/Private/ChaosVisualDebuggerMain.cpp
@@ -46,7 +46,7 @@ int32 RunChaosVisualDebugger(const TCHAR* CommandLine)
 	FQueuedThreadPool::OverrideStackSize = 256 * 1024;
 
 	const FText AppName = NSLOCTEXT("ChaosVisualDebugger", "ChaosVisualDebuggerSplashText", "Chaos Visual Debugger");
-	FPlatformSplash::SetSplashText(SplashTextType::GameName, *AppName.ToString());
+	FPlatformSplash::SetSplashText(SplashTextType::GameName, AppName);
 
 	FString Command;
 	const bool bIsRunningCommand = FParse::Value(CommandLine, TEXT("-RUN="), Command);
diff --git a/Engine/Source/Programs/LiveLinkHub/Source/LiveLinkHubLauncher/Private/LiveLinkHubRun.cpp b/Engine/Source/Programs/LiveLinkHub/Source/LiveLinkHubLauncher/Private/LiveLinkHubRun.cpp
--- a/Engine/Source/Programs/LiveLinkHub/Source/LiveLinkHubLauncher/Private/LiveLinkHubRun.cpp
+++ b/Engine/Source/Programs/LiveLinkHub/Source/LiveLinkHubLauncher/Private/LiveLinkHubRun.cpp
@@ -34,7 +34,7 @@ int32 RunLiveLinkHub(const TCHAR* CommandLine)
 	FCommandLine::Set(CommandLine);
 
 	const FText AppName = NSLOCTEXT("LiveLinkHub", "SplashTextName", "Live Link Hub");
-	FPlatformSplash::SetSplashText(SplashTextType::GameName, *AppName.ToString());
+	FPlatformSplash::SetSplashText(SplashTextType::GameName, AppName);
 
 #if !UE_BUILD_SHIPPING
 	if (FParse::Param(CommandLine, TEXT("WaitForDebugger")))
