[ChaosVD] Fix for an issue where CVD standalone couldn't start live recordings in some targets.

The problem is that by default, the code that launches Unreal Trace Server is compiled out in UE programs, therefore if the target client or server does not start the trace server themselves, CVD standalone can't start a live trace.

This CL enables the auto-launch trace server code in the CVD standalone program

diff --git a/Engine/Source/Programs/ChaosVisualDebugger/ChaosVisualDebugger.Target.cs b/Engine/Source/Programs/ChaosVisualDebugger/ChaosVisualDebugger.Target.cs
--- a/Engine/Source/Programs/ChaosVisualDebugger/ChaosVisualDebugger.Target.cs
+++ b/Engine/Source/Programs/ChaosVisualDebugger/ChaosVisualDebugger.Target.cs
@@ -61,5 +61,8 @@ public class ChaosVisualDebuggerTarget : TargetRules
 		// this will allow shader compiling to work based on whether or not the shaders directory is present
 		// to determine if we should allow shader compilation
 		GlobalDefinitions.Add("UE_ALLOW_SHADER_COMPILING_BASED_ON_SHADER_DIRECTORY_EXISTENCE=1");
+		
+		GlobalDefinitions.Add("UE_TRACE_SERVER_LAUNCH_ENABLED=1");
+		GlobalDefinitions.Add("UE_TRACE_SERVER_CONTROLS_ENABLED=1");
 	}
 }
